Test target that can build against Engine with editor only data, not cooked data required, does not compile against editor.
#rnx

diff --git a/Engine/Source/Programs/JustEngineTest/JustEngineTest.Build.cs b/Engine/Source/Programs/JustEngineTest/JustEngineTest.Build.cs
new file mode 100644
--- /dev/null
+++ b/Engine/Source/Programs/JustEngineTest/JustEngineTest.Build.cs
@@ -0,0 +1,14 @@
+// Copyright Epic Games, Inc. All Rights Reserved.
+using System;
+using UnrealBuildTool;
+
+public class JustEngineTest : TestModuleRules
+{
+	public JustEngineTest(ReadOnlyTargetRules Target) : base(Target)
+	{
+		PrivateDependencyModuleNames.AddRange(
+			new string[] {
+				"Engine"
+			});
+	}
+}
\ No newline at end of file
diff --git a/Engine/Source/Programs/JustEngineTest/JustEngineTest.Target.cs b/Engine/Source/Programs/JustEngineTest/JustEngineTest.Target.cs
new file mode 100644
--- /dev/null
+++ b/Engine/Source/Programs/JustEngineTest/JustEngineTest.Target.cs
@@ -0,0 +1,19 @@
+// Copyright Epic Games, Inc. All Rights Reserved.
+
+using System.IO;
+using UnrealBuildTool;
+
+[SupportedPlatforms(UnrealPlatformClass.All)]
+public class JustEngineTestTarget : TestTargetRules
+{
+	public JustEngineTestTarget(TargetInfo Target) : base(Target)
+	{
+		IncludeOrderVersion = EngineIncludeOrderVersion.Latest;
+		
+		bBuildWithEditorOnlyData = true;
+		bBuildRequiresCookedData = false;
+		
+		bCompileAgainstEditor = false;
+		bCompileAgainstEngine = true;
+	}
+}
diff --git a/Engine/Source/Programs/JustEngineTest/Private/SampleTest.cpp b/Engine/Source/Programs/JustEngineTest/Private/SampleTest.cpp
new file mode 100644
--- /dev/null
+++ b/Engine/Source/Programs/JustEngineTest/Private/SampleTest.cpp
@@ -0,0 +1,8 @@
+// Copyright Epic Games, Inc. All Rights Reserved.
+
+#include "TestHarness.h"
+
+TEST_CASE("Engine Sample", "[Engine]")
+{
+	CHECK(true);
+}
